Added GameFeaturesToolset plugin

Adds a new toolset that exposes Game Feature Plugin operations to the AI Toolset Registry:

- ListGameFeatures: returns a sorted list of all discovered GFPs
- FindGameFeatureData: loads the UGameFeatureData asset for a named plugin
- GetPluginName / GetActions: inspects an existing GameFeatureData asset
- CreateGameFeaturePlugin: creates the directory tree, writes the .uplugin, mounts the plugin, and saves the initial GameFeatureData asset

#Synced-CL 52233446

diff --git a/Engine/Plugins/Experimental/Toolsets/GameFeaturesToolset/GameFeaturesToolset.uplugin b/Engine/Plugins/Experimental/Toolsets/GameFeaturesToolset/GameFeaturesToolset.uplugin
new file mode 100644
--- /dev/null
+++ b/Engine/Plugins/Experimental/Toolsets/GameFeaturesToolset/GameFeaturesToolset.uplugin
@@ -0,0 +1,30 @@
+{
+	"FileVersion": 3,
+	"Version": 1,
+	"VersionName": "1.0",
+	"FriendlyName": "Game Features Toolset",
+	"Description": "Toolset for listing, inspecting, and creating Game Feature Plugins via the AI Toolset Registry.",
+	"Category": "Other",
+	"EditorOnly": true,
+	"IsExperimentalVersion": true,
+	"NoRedist": true,
+	"EnabledByDefault": false,
+	"Modules": [
+		{
+			"Name": "GameFeaturesToolset",
+			"Type": "Editor",
+			"LoadingPhase": "PostEngineInit",
+			"TargetAllowList": ["Editor"]
+		}
+	],
+	"Plugins": [
+		{
+			"Name": "ToolsetRegistry",
+			"Enabled": true
+		},
+		{
+			"Name": "GameFeatures",
+			"Enabled": true
+		}
+	]
+}
diff --git a/Engine/Plugins/Experimental/Toolsets/GameFeaturesToolset/Source/GameFeaturesToolset/GameFeaturesToolset.Build.cs b/Engine/Plugins/Experimental/Toolsets/GameFeaturesToolset/Source/GameFeaturesToolset/GameFeaturesToolset.Build.cs
new file mode 100644
--- /dev/null
+++ b/Engine/Plugins/Experimental/Toolsets/GameFeaturesToolset/Source/GameFeaturesToolset/GameFeaturesToolset.Build.cs
@@ -0,0 +1,33 @@
+// Copyright Epic Games, Inc. All Rights Reserved.
+
+using UnrealBuildTool;
+
+public class GameFeaturesToolset : ModuleRules
+{
+	public GameFeaturesToolset(ReadOnlyTargetRules Target) : base(Target)
+	{
+		PCHUsage = ModuleRules.PCHUsageMode.UseExplicitOrSharedPCHs;
+		IWYUSupport = IWYUSupport.Full;
+
+		PublicDependencyModuleNames.AddRange(
+			new string[]
+			{
+				"Core",
+			}
+			);
+
+		PrivateDependencyModuleNames.AddRange(
+			new string[]
+			{
+				"AssetRegistry",
+				"CoreUObject",
+				"Engine",
+				"GameFeatures",
+				"Kismet",
+				"Projects",
+				"ToolsetRegistry",
+				"UnrealEd",
+			}
+			);
+	}
+}
diff --git a/Engine/Plugins/Experimental/Toolsets/GameFeaturesToolset/Source/GameFeaturesToolset/Private/GameFeaturesToolset.cpp b/Engine/Plugins/Experimental/Toolsets/GameFeaturesToolset/Source/GameFeaturesToolset/Private/GameFeaturesToolset.cpp
new file mode 100644
--- /dev/null
+++ b/Engine/Plugins/Experimental/Toolsets/GameFeaturesToolset/Source/GameFeaturesToolset/Private/GameFeaturesToolset.cpp
@@ -0,0 +1,293 @@
+// Copyright Epic Games, Inc. All Rights Reserved.
+
+#include "GameFeaturesToolset.h"
+
+#include "AssetRegistry/AssetRegistryModule.h"
+#include "GameFeatureAction.h"
+#include "GameFeatureData.h"
+#include "HAL/PlatformFileManager.h"
+#include "Interfaces/IPluginManager.h"
+#include "Kismet/KismetSystemLibrary.h"
+#include "Misc/PackageName.h"
+#include "Misc/Paths.h"
+#include "Misc/ScopeExit.h"
+#include "PluginDescriptor.h"
+#include "UObject/Package.h"
+#include "UObject/SavePackage.h"
+#include "UObject/UObjectGlobals.h"
+
+#include UE_INLINE_GENERATED_CPP_BY_NAME(GameFeaturesToolset)
+
+namespace
+{
+	bool IsGameFeaturePlugin(const IPlugin& Plugin)
+	{
+		for (const FPluginReferenceDescriptor& Dep : Plugin.GetDescriptor().Plugins)
+		{
+			if (Dep.Name == TEXT("GameFeatures") && Dep.bEnabled)
+			{
+				return true;
+			}
+		}
+		return false;
+	}
+
+	// Returns true if Name is a valid plugin identifier: non-empty, letters/digits/underscores only.
+	bool IsValidPluginName(const FString& Name)
+	{
+		if (Name.IsEmpty())
+		{
+			return false;
+		}
+		for (TCHAR Ch : Name)
+		{
+			if (!FChar::IsAlnum(Ch) && Ch != TEXT('_'))
+			{
+				return false;
+			}
+		}
+		return true;
+	}
+}
+
+TArray<FString> UGameFeaturesToolset::ListGameFeatures()
+{
+	TArray<FString> Names;
+	for (const TSharedRef<IPlugin>& Plugin : IPluginManager::Get().GetDiscoveredPlugins())
+	{
+		if (IsGameFeaturePlugin(Plugin.Get()))
+		{
+			Names.Add(Plugin->GetName());
+		}
+	}
+	Names.Sort();
+	return Names;
+}
+
+UGameFeatureData* UGameFeaturesToolset::FindGameFeatureData(const FString& PluginName)
+{
+	if (PluginName.IsEmpty())
+	{
+		UKismetSystemLibrary::RaiseScriptError(TEXT("PluginName cannot be empty."));
+		return nullptr;
+	}
+
+	TSharedPtr<IPlugin> Plugin = IPluginManager::Get().FindPlugin(PluginName);
+	if (!Plugin.IsValid())
+	{
+		UKismetSystemLibrary::RaiseScriptError(
+			FString::Printf(TEXT("Plugin '%s' was not found."), *PluginName));
+		return nullptr;
+	}
+
+	if (!IsGameFeaturePlugin(*Plugin))
+	{
+		UKismetSystemLibrary::RaiseScriptError(
+			FString::Printf(TEXT("Plugin '%s' is not a Game Feature Plugin."), *PluginName));
+		return nullptr;
+	}
+
+	const FString AssetPath = FString::Printf(TEXT("/%s/%s"), *PluginName, *PluginName);
+	UGameFeatureData* Data = Cast<UGameFeatureData>(
+		StaticLoadObject(UGameFeatureData::StaticClass(), nullptr, *AssetPath));
+
+	if (!Data)
+	{
+		UKismetSystemLibrary::RaiseScriptError(
+			FString::Printf(TEXT("Could not load GameFeatureData at '%s'."), *AssetPath));
+		return nullptr;
+	}
+
+	return Data;
+}
+
+FString UGameFeaturesToolset::GetPluginName(UGameFeatureData* Data)
+{
+	if (!Data)
+	{
+		UKismetSystemLibrary::RaiseScriptError(TEXT("Data is null."));
+		return FString();
+	}
+
+	FString OutName;
+	UGameFeatureData::GetPluginName(Data, OutName);
+	return OutName;
+}
+
+TArray<UGameFeatureAction*> UGameFeaturesToolset::GetActions(UGameFeatureData* Data)
+{
+	if (!Data)
+	{
+		UKismetSystemLibrary::RaiseScriptError(TEXT("Data is null."));
+		return {};
+	}
+
+	return Data->GetActions();
+}
+
+UGameFeatureData* UGameFeaturesToolset::CreateGameFeaturePlugin(
+	const FString& PluginName, const FString& Description)
+{
+	// Step 1: Validate inputs.
+	if (!IsValidPluginName(PluginName))
+	{
+		UKismetSystemLibrary::RaiseScriptError(
+			FString::Printf(
+				TEXT("PluginName '%s' is invalid. Must be non-empty and contain only "
+				     "letters, digits, and underscores."),
+				*PluginName));
+		return nullptr;
+	}
+
+	if (IPluginManager::Get().FindPlugin(PluginName).IsValid())
+	{
+		UKismetSystemLibrary::RaiseScriptError(
+			FString::Printf(TEXT("A plugin named '%s' already exists."), *PluginName));
+		return nullptr;
+	}
+
+	// Step 2: Create the directory tree.
+	const FString PluginDir = FPaths::ProjectPluginsDir() / TEXT("GameFeatures") / PluginName;
+	const FString ContentDir = PluginDir / TEXT("Content");
+	const FString UpluginPath = PluginDir / (PluginName + TEXT(".uplugin"));
+
+	// bCommitted makes all guards no-ops on success. Guards are registered incrementally
+	// (each immediately after the operation it reverses), so only operations that succeeded
+	// are unwound. LIFO teardown order: GC-data -> unmount -> remove-from-list -> delete-directory.
+	bool bCommitted = false;
+
+	IPlatformFile& PlatformFile = FPlatformFileManager::Get().GetPlatformFile();
+	if (!PlatformFile.CreateDirectoryTree(*ContentDir))
+	{
+		UKismetSystemLibrary::RaiseScriptError(
+			FString::Printf(TEXT("Failed to create directory '%s'."), *ContentDir));
+		return nullptr;
+	}
+	ON_SCOPE_EXIT
+	{
+		if (!bCommitted)
+		{
+			PlatformFile.DeleteDirectoryRecursively(*PluginDir);
+		}
+	};
+
+	// Step 3: Write the .uplugin file.
+	FPluginDescriptor Descriptor;
+	Descriptor.FriendlyName = PluginName;
+	Descriptor.Description = Description;
+	Descriptor.Version = 1;
+	Descriptor.VersionName = TEXT("1.0");
+	Descriptor.Category = TEXT("GameFeatures");
+	Descriptor.bExplicitlyLoaded = true;
+	Descriptor.bCanContainContent = true;
+
+	FPluginReferenceDescriptor GameFeaturesRef;
+	GameFeaturesRef.Name = TEXT("GameFeatures");
+	GameFeaturesRef.bEnabled = true;
+	Descriptor.Plugins.Add(GameFeaturesRef);
+
+	FText SaveFailReason;
+	if (!Descriptor.Save(UpluginPath, SaveFailReason))
+	{
+		UKismetSystemLibrary::RaiseScriptError(
+			FString::Printf(TEXT("Failed to write .uplugin: %s"), *SaveFailReason.ToString()));
+		return nullptr;
+	}
+
+	// Step 4: Mount the plugin.
+	if (!IPluginManager::Get().AddToPluginsList(UpluginPath))
+	{
+		UKismetSystemLibrary::RaiseScriptError(
+			FString::Printf(TEXT("Failed to add '%s' to the plugin list."), *PluginName));
+		return nullptr;
+	}
+	ON_SCOPE_EXIT
+	{
+		if (!bCommitted)
+		{
+			FText RemoveReason;
+			IPluginManager::Get().RemoveFromPluginsList(UpluginPath, &RemoveReason);
+		}
+	};
+
+	if (!IPluginManager::Get().MountExplicitlyLoadedPlugin(PluginName))
+	{
+		UKismetSystemLibrary::RaiseScriptError(
+			FString::Printf(TEXT("Failed to mount plugin '%s'."), *PluginName));
+		return nullptr;
+	}
+	ON_SCOPE_EXIT
+	{
+		if (!bCommitted)
+		{
+			FText UnmountReason;
+			IPluginManager::Get().UnmountExplicitlyLoadedPlugin(PluginName, &UnmountReason);
+		}
+	};
+
+	// Step 5: Create the UGameFeatureData asset in-memory.
+	const FString PackagePath = FString::Printf(TEXT("/%s/%s"), *PluginName, *PluginName);
+	UPackage* NewPackage = CreatePackage(*PackagePath);
+	if (!NewPackage)
+	{
+		UKismetSystemLibrary::RaiseScriptError(
+			FString::Printf(TEXT("Failed to create package at '%s'."), *PackagePath));
+		return nullptr;
+	}
+	// Declared here (before the scope exit) so the lambda can reference it.
+	// Initialized after NewObject below.
+	UGameFeatureData* Data = nullptr;
+
+	ON_SCOPE_EXIT
+	{
+		if (!bCommitted)
+		{
+			// RF_Standalone prevents GC even when a package is marked as garbage, so clear
+			// it from Data explicitly before marking both for collection.
+			if (Data)
+			{
+				Data->ClearFlags(RF_Standalone);
+				Data->MarkAsGarbage();
+			}
+			NewPackage->MarkAsGarbage();
+		}
+	};
+
+	Data = NewObject<UGameFeatureData>(
+		NewPackage, FName(*PluginName), RF_Public | RF_Standalone);
+	if (!Data)
+	{
+		UKismetSystemLibrary::RaiseScriptError(
+			FString::Printf(
+				TEXT("Failed to create GameFeatureData object in package '%s'."), *PackagePath));
+		return nullptr;
+	}
+
+	// Step 6: Save to disk.
+	FString PackageFilename;
+	if (!FPackageName::TryConvertLongPackageNameToFilename(
+			PackagePath, PackageFilename, FPackageName::GetAssetPackageExtension()))
+	{
+		UKismetSystemLibrary::RaiseScriptError(
+			FString::Printf(
+				TEXT("Failed to resolve file path for package '%s'."), *PackagePath));
+		return nullptr;
+	}
+
+	FSavePackageArgs SaveArgs;
+	SaveArgs.TopLevelFlags = RF_Standalone;
+	SaveArgs.Error = GLog;
+
+	if (!UPackage::SavePackage(NewPackage, Data, *PackageFilename, SaveArgs))
+	{
+		UKismetSystemLibrary::RaiseScriptError(
+			FString::Printf(TEXT("Failed to save package to '%s'."), *PackageFilename));
+		return nullptr;
+	}
+
+	// Step 7: Notify the asset registry now that the asset is durably saved.
+	FAssetRegistryModule::AssetCreated(Data);
+
+	bCommitted = true;
+	return Data;
+}
diff --git a/Engine/Plugins/Experimental/Toolsets/GameFeaturesToolset/Source/GameFeaturesToolset/Private/GameFeaturesToolset.h b/Engine/Plugins/Experimental/Toolsets/GameFeaturesToolset/Source/GameFeaturesToolset/Private/GameFeaturesToolset.h
new file mode 100644
--- /dev/null
+++ b/Engine/Plugins/Experimental/Toolsets/GameFeaturesToolset/Source/GameFeaturesToolset/Private/GameFeaturesToolset.h
@@ -0,0 +1,60 @@
+// Copyright Epic Games, Inc. All Rights Reserved.
+
+#pragma once
+
+#include "CoreMinimal.h"
+#include "ToolsetRegistry/UToolsetRegistry.h"
+
+#include "GameFeaturesToolset.generated.h"
+
+class UGameFeatureAction;
+class UGameFeatureData;
+
+/// Provides tools for listing, inspecting, and creating Game Feature Plugins.
+UCLASS(BlueprintType, Hidden)
+class UGameFeaturesToolset : public UToolsetDefinition
+{
+	GENERATED_BODY()
+
+public:
+	/// Returns the names of all Game Feature Plugins currently known to the plugin manager.
+	/// A plugin is considered a Game Feature Plugin if it declares "GameFeatures" as a dependency.
+	/// @return A sorted list of plugin names.
+	UFUNCTION(BlueprintCallable, Category = "GameFeatures")
+	static TArray<FString> ListGameFeatures();
+
+	/// Finds and loads the UGameFeatureData asset for the named Game Feature Plugin.
+	/// @param PluginName The name of the plugin (e.g. "ShooterCore").
+	/// @return The plugin's GameFeatureData asset.
+	///   Raises a script error if the plugin does not exist, is not a Game Feature Plugin,
+	///   or the data asset cannot be loaded.
+	UFUNCTION(BlueprintCallable, Category = "GameFeatures")
+	static UGameFeatureData* FindGameFeatureData(const FString& PluginName);
+
+	/// Returns the plugin name associated with a GameFeatureData asset.
+	/// @param Data The GameFeatureData asset to query.
+	/// @return The plugin name string.
+	///   Raises a script error if Data is null.
+	UFUNCTION(BlueprintCallable, Category = "GameFeatures")
+	static FString GetPluginName(UGameFeatureData* Data);
+
+	/// Returns all GameFeatureActions defined in a GameFeatureData asset.
+	/// @param Data The GameFeatureData asset to query.
+	/// @return The list of actions.
+	///   Raises a script error if Data is null.
+	UFUNCTION(BlueprintCallable, Category = "GameFeatures")
+	static TArray<UGameFeatureAction*> GetActions(UGameFeatureData* Data);
+
+	/// Creates a new Game Feature Plugin with the given name and description.
+	/// This creates the plugin directory structure, writes the .uplugin file, mounts the
+	/// plugin, and creates the initial GameFeatureData asset.
+	/// This should ONLY be called after getting explicit direction or permission from the user.
+	/// @param PluginName The name for the new plugin. Must be unique, non-empty, and contain
+	///   only letters, digits, and underscores (e.g. "ShooterCore", not "My-Plugin").
+	/// @param Description An optional description for the plugin.
+	/// @return The newly created GameFeatureData asset.
+	///   Raises a script error if the plugin already exists or any creation step fails.
+	UFUNCTION(BlueprintCallable, Category = "GameFeatures")
+	static UGameFeatureData* CreateGameFeaturePlugin(
+		const FString& PluginName, const FString& Description);
+};
diff --git a/Engine/Plugins/Experimental/Toolsets/GameFeaturesToolset/Source/GameFeaturesToolset/Private/Module.cpp b/Engine/Plugins/Experimental/Toolsets/GameFeaturesToolset/Source/GameFeaturesToolset/Private/Module.cpp
new file mode 100644
--- /dev/null
+++ b/Engine/Plugins/Experimental/Toolsets/GameFeaturesToolset/Source/GameFeaturesToolset/Private/Module.cpp
@@ -0,0 +1,22 @@
+// Copyright Epic Games, Inc. All Rights Reserved.
+
+#include "GameFeaturesToolset.h"
+#include "Modules/ModuleInterface.h"
+#include "Modules/ModuleManager.h"
+#include "ToolsetRegistry/UToolsetRegistry.h"
+
+class FGameFeaturesToolsetModule : public IModuleInterface
+{
+public:
+	virtual void StartupModule() override
+	{
+		UToolsetRegistry::RegisterToolsetClass(UGameFeaturesToolset::StaticClass());
+	}
+
+	virtual void ShutdownModule() override
+	{
+		UToolsetRegistry::UnregisterToolsetClass(UGameFeaturesToolset::StaticClass());
+	}
+};
+
+IMPLEMENT_MODULE(FGameFeaturesToolsetModule, GameFeaturesToolset);
diff --git a/Engine/Plugins/Experimental/Toolsets/GameFeaturesToolset/Source/GameFeaturesToolset/Private/Tests/GameFeatureToolsetTest.cpp b/Engine/Plugins/Experimental/Toolsets/GameFeaturesToolset/Source/GameFeaturesToolset/Private/Tests/GameFeatureToolsetTest.cpp
new file mode 100644
--- /dev/null
+++ b/Engine/Plugins/Experimental/Toolsets/GameFeaturesToolset/Source/GameFeaturesToolset/Private/Tests/GameFeatureToolsetTest.cpp
@@ -0,0 +1,216 @@
+// Copyright Epic Games, Inc. All Rights Reserved.
+
+#include "AssetRegistry/AssetRegistryModule.h"
+#include "GameFeatureAction.h"
+#include "GameFeatureData.h"
+#include "GameFeaturesToolset.h"
+#include "Interfaces/IPluginManager.h"
+#include "Misc/AutomationTest.h"
+#include "UObject/StrongObjectPtr.h"
+
+#if WITH_DEV_AUTOMATION_TESTS
+
+namespace
+{
+	const EAutomationTestFlags TestFlags =
+		EAutomationTestFlags::EditorContext |
+		EAutomationTestFlags::ProductFilter |
+		EAutomationTestFlags::CriticalPriority;
+
+	// Returns the name of the first GFP that has a registered GameFeatureData asset in the
+	// asset registry, or an empty string if none is found in this build configuration.
+	// Used to guard happy-path tests that require a live GFP without triggering
+	// RaiseScriptError on a miss.
+	FString FindAnyRegisteredGFPName()
+	{
+		const TSet<FString> GFPNames(UGameFeaturesToolset::ListGameFeatures());
+
+		IAssetRegistry& AssetRegistry =
+			FModuleManager::LoadModuleChecked<FAssetRegistryModule>(TEXT("AssetRegistry")).Get();
+
+		TArray<FAssetData> Assets;
+		AssetRegistry.GetAssetsByClass(
+			UGameFeatureData::StaticClass()->GetClassPathName(), Assets, /*bSearchSubClasses=*/true);
+
+		for (const FAssetData& Asset : Assets)
+		{
+			TArray<FString> Parts;
+			Asset.PackageName.ToString().ParseIntoArray(Parts, TEXT("/"), /*InCullEmpty=*/true);
+			if (Parts.Num() > 0 && GFPNames.Contains(Parts[0]))
+			{
+				return Parts[0];
+			}
+		}
+		return FString();
+	}
+}
+
+BEGIN_DEFINE_SPEC(FListGameFeaturesSpec, "AI.Toolsets.GameFeatureToolset.ListGameFeatures",
+	TestFlags)
+END_DEFINE_SPEC(FListGameFeaturesSpec)
+
+void FListGameFeaturesSpec::Define()
+{
+	It(TEXT("returns a sorted list"), [this]()
+	{
+		const TArray<FString> Names = UGameFeaturesToolset::ListGameFeatures();
+		// Guard against a vacuous pass: this project always has at least one Game Feature Plugin.
+		TestTrue(TEXT("At least one Game Feature Plugin exists"), Names.Num() > 0);
+		for (int32 i = 1; i < Names.Num(); ++i)
+		{
+			TestTrue(TEXT("Sorted"), Names[i - 1] <= Names[i]);
+		}
+	});
+
+	It(TEXT("does not include plugins without a GameFeatures dependency"), [this]()
+	{
+		const TArray<FString> Names = UGameFeaturesToolset::ListGameFeatures();
+		// ToolsetRegistry declares no GameFeatures dependency, so must not appear.
+		TestFalse(TEXT("ToolsetRegistry excluded"), Names.Contains(TEXT("ToolsetRegistry")));
+		// The GameFeatures plugin itself does not declare GameFeatures as a dependency.
+		TestFalse(TEXT("GameFeatures excluded"), Names.Contains(TEXT("GameFeatures")));
+	});
+
+	It(TEXT("only includes plugins that are discovered by the plugin manager"), [this]()
+	{
+		const TArray<FString> Names = UGameFeaturesToolset::ListGameFeatures();
+
+		TSet<FString> DiscoveredNames;
+		for (const TSharedRef<IPlugin>& Plugin : IPluginManager::Get().GetDiscoveredPlugins())
+		{
+			DiscoveredNames.Add(Plugin->GetName());
+		}
+
+		for (const FString& Name : Names)
+		{
+			TestTrue(
+				FString::Printf(TEXT("'%s' is a discovered plugin"), *Name),
+				DiscoveredNames.Contains(Name));
+		}
+	});
+}
+
+BEGIN_DEFINE_SPEC(FFindGameFeatureDataSpec, "AI.Toolsets.GameFeatureToolset.FindGameFeatureData",
+	TestFlags)
+END_DEFINE_SPEC(FFindGameFeatureDataSpec)
+
+void FFindGameFeatureDataSpec::Define()
+{
+	It(TEXT("returns null for an empty plugin name"), [this]()
+	{
+		TestNull(TEXT("Result"), UGameFeaturesToolset::FindGameFeatureData(TEXT("")));
+	});
+
+	It(TEXT("returns null for a plugin that does not exist"), [this]()
+	{
+		TestNull(TEXT("Result"),
+			UGameFeaturesToolset::FindGameFeatureData(TEXT("DoesNotExist_AutoTest_12345")));
+	});
+
+	It(TEXT("returns null for a plugin that is not a Game Feature Plugin"), [this]()
+	{
+		// ToolsetRegistry is a known plugin that is never a GFP.
+		TestNull(TEXT("Result"),
+			UGameFeaturesToolset::FindGameFeatureData(TEXT("ToolsetRegistry")));
+	});
+
+	It(TEXT("returns a valid asset for an existing Game Feature Plugin"), [this]()
+	{
+		const FString PluginName = FindAnyRegisteredGFPName();
+		if (PluginName.IsEmpty())
+		{
+			AddInfo(TEXT("No registered GameFeatureData asset has a package root matching a "
+			             "known GFP name; skipping happy-path assertion."));
+			return;
+		}
+
+		UGameFeatureData* Data = UGameFeaturesToolset::FindGameFeatureData(PluginName);
+		TestNotNull(TEXT("Data is non-null"), Data);
+	});
+}
+
+BEGIN_DEFINE_SPEC(FGetPluginNameSpec, "AI.Toolsets.GameFeatureToolset.GetPluginName", TestFlags)
+END_DEFINE_SPEC(FGetPluginNameSpec)
+
+void FGetPluginNameSpec::Define()
+{
+	It(TEXT("returns an empty string for a null asset"), [this]()
+	{
+		TestEqual(TEXT("Result"), UGameFeaturesToolset::GetPluginName(nullptr), FString());
+	});
+
+	It(TEXT("round-trips the plugin name through FindGameFeatureData"), [this]()
+	{
+		const FString ExpectedName = FindAnyRegisteredGFPName();
+		if (ExpectedName.IsEmpty())
+		{
+			AddInfo(TEXT("No registered GameFeatureData asset has a package root matching a "
+			             "known GFP name; skipping round-trip assertion."));
+			return;
+		}
+
+		UGameFeatureData* Data = UGameFeaturesToolset::FindGameFeatureData(ExpectedName);
+		if (!TestNotNull(TEXT("Data is non-null"), Data))
+		{
+			return;
+		}
+
+		TestEqual(TEXT("Plugin name round-trips"), UGameFeaturesToolset::GetPluginName(Data), ExpectedName);
+	});
+}
+
+BEGIN_DEFINE_SPEC(FGetActionsSpec, "AI.Toolsets.GameFeatureToolset.GetActions", TestFlags)
+END_DEFINE_SPEC(FGetActionsSpec)
+
+void FGetActionsSpec::Define()
+{
+	It(TEXT("returns an empty array for a null asset"), [this]()
+	{
+		TestTrue(TEXT("Empty"), UGameFeaturesToolset::GetActions(nullptr).IsEmpty());
+	});
+
+	It(TEXT("returns an empty array for a new data asset with no actions"), [this]()
+	{
+		TStrongObjectPtr<UGameFeatureData> Data(
+			NewObject<UGameFeatureData>(GetTransientPackage(), NAME_None, RF_Transient));
+
+		if (!TestNotNull(TEXT("Data"), Data.Get())) return;
+
+		TestTrue(TEXT("Empty"), UGameFeaturesToolset::GetActions(Data.Get()).IsEmpty());
+	});
+}
+
+BEGIN_DEFINE_SPEC(FCreateGameFeaturePluginSpec,
+	"AI.Toolsets.GameFeatureToolset.CreateGameFeaturePlugin", TestFlags)
+END_DEFINE_SPEC(FCreateGameFeaturePluginSpec)
+
+void FCreateGameFeaturePluginSpec::Define()
+{
+	// The end-to-end creation path (successful plugin creation, mount, and asset save) is not
+	// covered by an integration test because it creates persistent filesystem and plugin-manager
+	// state that is difficult to roll back cleanly in the automation framework. Input-validation
+	// paths (below) are unit-tested; the happy path is covered by manual verification.
+
+	It(TEXT("returns null for an empty plugin name"), [this]()
+	{
+		TestNull(TEXT("Result"),
+			UGameFeaturesToolset::CreateGameFeaturePlugin(TEXT(""), TEXT("")));
+	});
+
+	It(TEXT("returns null for a plugin name containing path separators"), [this]()
+	{
+		TestNull(TEXT("Result"),
+			UGameFeaturesToolset::CreateGameFeaturePlugin(TEXT("../../etc/passwd"), TEXT("")));
+		TestNull(TEXT("Result"),
+			UGameFeaturesToolset::CreateGameFeaturePlugin(TEXT("foo/bar"), TEXT("")));
+	});
+
+	It(TEXT("returns null when the plugin already exists"), [this]()
+	{
+		// "GameFeatures" is an engine plugin that always exists.
+		TestNull(TEXT("Result"),
+			UGameFeaturesToolset::CreateGameFeaturePlugin(TEXT("GameFeatures"), TEXT("")));
+	});
+}
+
+#endif  // WITH_DEV_AUTOMATION_TESTS
